integrate emission color and smoke separately in ray march step

diff --git a/Engine/Plugins/FX/NiagaraFluids/Shaders/NiagaraFluids.ush b/Engine/Plugins/FX/NiagaraFluids/Shaders/NiagaraFluids.ush
--- a/Engine/Plugins/FX/NiagaraFluids/Shaders/NiagaraFluids.ush
+++ b/Engine/Plugins/FX/NiagaraFluids/Shaders/NiagaraFluids.ush
@@ -10,8 +10,6 @@ NiagaraDataInterfacePhysicsField.ush
  * Rendering code for simulation stages
  * -----------------------------------------------------------------
  */
-
-
 #define MARCH_SINGLE_RAY(RayStart, RayDir, UnitStepSize, LocalStepSize, NumSteps, Volume, RetV) \
 RetV = float4(0,0,0,0); \
 float Transmittance = 1; \
@@ -23,13 +21,45 @@ for (int i = 0; i < Steps && Transmittance > 1e-5; i++) \
 float Extinction; \
 float Scattering; \
 float3 Luminance; \
-Volume.ComputeStep(RayStart, RayDir, Extinction, Scattering, Luminance); \
+float EmissionExtinction = 0; \
+float3 EmissionColor = float3(0,0,0); \
+Volume.ComputeStep(RayStart, RayDir, Extinction, Scattering, Luminance, EmissionExtinction, EmissionColor); \
 Luminance *= Scattering; \
 float ClampedExtinction = max(Extinction, 1e-6); \
 float StepTransmittance = exp(-Extinction * LocalStepSize); \
 RetV.rgb += Transmittance * (Luminance - Luminance * StepTransmittance) / ClampedExtinction; \
 Transmittance *= StepTransmittance; \
+Luminance *= Scattering; \
+if(i == Steps-2) UnitStepSize, LocalStepSize *= frac(NumSteps); \
+RayStart += RayDir * UnitStepSize; \
+} \
+RetV.a = 1. - Transmittance;
+
+#define MARCH_SINGLE_RAY_EMISSION(RayStart, RayDir, UnitStepSize, LocalStepSize, NumSteps, Volume, RetV) \
+RetV = float4(0,0,0,0); \
+float Transmittance = 1; \
+int Steps = ceil(NumSteps); \
+if(Steps >= 1.0 ) Steps += 1.0; \
+[loop] \
+for (int i = 0; i < Steps && Transmittance > 1e-5; i++) \
+{ \
+float Extinction; \
+float Scattering; \
+float3 Luminance; \
+float EmissionExtinction = 0; \
+float3 EmissionColor = float3(0,0,0); \
+Volume.ComputeStep(RayStart, RayDir, Extinction, Scattering, Luminance, EmissionExtinction, EmissionColor); \
+float ClampedEmissionExtinction = max(EmissionExtinction, 1e-6); \
+float StepEmissionTransmittance = exp(-EmissionExtinction * LocalStepSize); \
+RetV.rgb += Transmittance * (EmissionColor - EmissionColor * StepEmissionTransmittance) / ClampedEmissionExtinction; \
+Transmittance *= StepEmissionTransmittance; \
+Luminance *= Scattering; \
+float ClampedExtinction = max(Extinction, 1e-6); \
+float StepTransmittance = exp(-Extinction * LocalStepSize); \
+RetV.rgb += Transmittance * (Luminance - Luminance * StepTransmittance) / ClampedExtinction; \
+Transmittance *= StepTransmittance; \
+Luminance *= Scattering; \
 if(i == Steps-2) UnitStepSize, LocalStepSize *= frac(NumSteps); \
 RayStart += RayDir * UnitStepSize; \
 } \
-RetV.a = 1. - Transmittance; \
+RetV.a = 1. - Transmittance;
