fix non emission base shader

diff --git a/Engine/Plugins/FX/NiagaraFluids/Shaders/NiagaraFluids.ush b/Engine/Plugins/FX/NiagaraFluids/Shaders/NiagaraFluids.ush
--- a/Engine/Plugins/FX/NiagaraFluids/Shaders/NiagaraFluids.ush
+++ b/Engine/Plugins/FX/NiagaraFluids/Shaders/NiagaraFluids.ush
@@ -21,9 +21,7 @@ for (int i = 0; i < Steps && Transmittance > 1e-5; i++) \
 float Extinction; \
 float Scattering; \
 float3 Luminance; \
-float EmissionExtinction = 0; \
-float3 EmissionColor = float3(0,0,0); \
-Volume.ComputeStep(RayStart, RayDir, Extinction, Scattering, Luminance, EmissionExtinction, EmissionColor); \
+Volume.ComputeStep(RayStart, RayDir, Extinction, Scattering, Luminance); \
 Luminance *= Scattering; \
 float ClampedExtinction = max(Extinction, 1e-6); \
 float StepTransmittance = exp(-Extinction * LocalStepSize); \
