fix copy/paste error

diff --git a/Engine/Plugins/FX/NiagaraFluids/Shaders/NiagaraFluids.ush b/Engine/Plugins/FX/NiagaraFluids/Shaders/NiagaraFluids.ush
--- a/Engine/Plugins/FX/NiagaraFluids/Shaders/NiagaraFluids.ush
+++ b/Engine/Plugins/FX/NiagaraFluids/Shaders/NiagaraFluids.ush
@@ -27,7 +27,6 @@ float ClampedExtinction = max(Extinction, 1e-6); \
 float StepTransmittance = exp(-Extinction * LocalStepSize); \
 RetV.rgb += Transmittance * (Luminance - Luminance * StepTransmittance) / ClampedExtinction; \
 Transmittance *= StepTransmittance; \
-Luminance *= Scattering; \
 if(i == Steps-2) UnitStepSize, LocalStepSize *= frac(NumSteps); \
 RayStart += RayDir * UnitStepSize; \
 } \
@@ -56,7 +55,6 @@ float ClampedExtinction = max(Extinction, 1e-6); \
 float StepTransmittance = exp(-Extinction * LocalStepSize); \
 RetV.rgb += Transmittance * (Luminance - Luminance * StepTransmittance) / ClampedExtinction; \
 Transmittance *= StepTransmittance; \
-Luminance *= Scattering; \
 if(i == Steps-2) UnitStepSize, LocalStepSize *= frac(NumSteps); \
 RayStart += RayDir * UnitStepSize; \
 } \
