- remove experimental flag on plugin
- remove overlays on colored smoke example

diff --git a/Engine/Plugins/FX/NiagaraFluids/NiagaraFluids.uplugin b/Engine/Plugins/FX/NiagaraFluids/NiagaraFluids.uplugin
--- a/Engine/Plugins/FX/NiagaraFluids/NiagaraFluids.uplugin
+++ b/Engine/Plugins/FX/NiagaraFluids/NiagaraFluids.uplugin
@@ -12,7 +12,7 @@
 	"SupportURL": "",
 	"CanContainContent": true,
 	"IsBetaVersion": false,
-	"IsExperimentalVersion": true,
+	"IsExperimentalVersion": false,
 	"Installed": false,
 	"Modules": [
 		{
