update uplugin file to specify beta,
fix missing material error

diff --git a/Engine/Plugins/FX/NiagaraFluids/NiagaraFluids.uplugin b/Engine/Plugins/FX/NiagaraFluids/NiagaraFluids.uplugin
--- a/Engine/Plugins/FX/NiagaraFluids/NiagaraFluids.uplugin
+++ b/Engine/Plugins/FX/NiagaraFluids/NiagaraFluids.uplugin
@@ -1,17 +1,17 @@
 {
-	"FileVersion": 3,
+	"FileVersion": 1,
 	"Version": 1,
 	"VersionName": "1.0",
 	"FriendlyName": "NiagaraFluids",
-	"Description": "",
-	"Category": "Other",
+	"Description": "Fluid simulation toolkit for Niagara",
+	"Category": "FX",
 	"CreatedBy": "Epic Games, Inc.",
-	"CreatedByURL": "",
+	"CreatedByURL": "http://epicgames.com",
 	"DocsURL": "",
 	"MarketplaceURL": "",
 	"SupportURL": "",
 	"CanContainContent": true,
-	"IsBetaVersion": false,
+	"IsBetaVersion": true,
 	"IsExperimentalVersion": false,
 	"Installed": false,
 	"Modules": [
