Temporarily add dependency on HairStrands plugin for NiagaraFluids because of the use of the PhysicsAssetDI.

diff --git a/Engine/Plugins/FX/NiagaraFluids/NiagaraFluids.uplugin b/Engine/Plugins/FX/NiagaraFluids/NiagaraFluids.uplugin
--- a/Engine/Plugins/FX/NiagaraFluids/NiagaraFluids.uplugin
+++ b/Engine/Plugins/FX/NiagaraFluids/NiagaraFluids.uplugin
@@ -25,6 +25,10 @@
         {
             "Name": "Niagara",			
             "Enabled": true
+        },
+        {
+            "Name": "HairStrands",			
+            "Enabled": true
         }
 	]
 }
\ No newline at end of file
