Disable dependency on HairStrands plugin for the NiagaraFluids.  Old assets may now fail to run properly without manually enabling the HairStrands plugin for access to the NiagaraDataInterfacePhysicsAsset class.  Newer are no longer dependent on this asset, and in the future the class itself will be migrated to the Niagara plugin.

diff --git a/Engine/Plugins/FX/NiagaraFluids/NiagaraFluids.uplugin b/Engine/Plugins/FX/NiagaraFluids/NiagaraFluids.uplugin
--- a/Engine/Plugins/FX/NiagaraFluids/NiagaraFluids.uplugin
+++ b/Engine/Plugins/FX/NiagaraFluids/NiagaraFluids.uplugin
@@ -25,10 +25,6 @@
         {
             "Name": "Niagara",			
             "Enabled": true
-        },
-        {
-            "Name": "HairStrands",			
-            "Enabled": true
         }
 	]
 }
\ No newline at end of file
