NiagaraFluids: Add ShallowWater-WaterBody integration WIP

- Niagara system
- WaterBody context Helper

diff --git a/Engine/Plugins/FX/NiagaraFluids/NiagaraFluids.uplugin b/Engine/Plugins/FX/NiagaraFluids/NiagaraFluids.uplugin
--- a/Engine/Plugins/FX/NiagaraFluids/NiagaraFluids.uplugin
+++ b/Engine/Plugins/FX/NiagaraFluids/NiagaraFluids.uplugin
@@ -1,5 +1,5 @@
 {
-	"FileVersion": 1,
+	"FileVersion": 3,
 	"Version": 1,
 	"VersionName": "1.0",
 	"FriendlyName": "NiagaraFluids",
@@ -20,11 +20,15 @@
 			"Type": "Runtime",
 			"LoadingPhase": "Default"
 		}
-	],    
+	],
 	"Plugins": [
-        {
-            "Name": "Niagara",			
-            "Enabled": true
-        }
+		{
+			"Name": "Niagara",
+			"Enabled": true
+		},
+		{
+			"Name": "Water",
+			"Enabled": true
+		}
 	]
 }
\ No newline at end of file
diff --git a/Engine/Plugins/FX/NiagaraFluids/Source/NiagaraFluids/NiagaraFluids.Build.cs b/Engine/Plugins/FX/NiagaraFluids/Source/NiagaraFluids/NiagaraFluids.Build.cs
--- a/Engine/Plugins/FX/NiagaraFluids/Source/NiagaraFluids/NiagaraFluids.Build.cs
+++ b/Engine/Plugins/FX/NiagaraFluids/Source/NiagaraFluids/NiagaraFluids.Build.cs
@@ -7,8 +7,18 @@ namespace UnrealBuildTool.Rules
 		public NiagaraFluids(ReadOnlyTargetRules Target) : base(Target)
 		{
 			PublicDependencyModuleNames.Add("Core");
-			PrivateDependencyModuleNames.Add("RenderCore");
-			PrivateDependencyModuleNames.Add("Projects");
+			PrivateDependencyModuleNames.AddRange(
+            			new string[] {
+				            "Core",
+				            "CoreUObject",
+				            "Engine",
+				            "RenderCore",
+            				"Niagara",
+            				"Water",
+            				"DeveloperSettings",
+				            "Projects"
+            			}
+            		);
 		}
 	}
 }
diff --git a/Engine/Plugins/FX/NiagaraFluids/Source/NiagaraFluids/Private/ShallowWaterComponent.cpp b/Engine/Plugins/FX/NiagaraFluids/Source/NiagaraFluids/Private/ShallowWaterComponent.cpp
new file mode 100644
--- /dev/null
+++ b/Engine/Plugins/FX/NiagaraFluids/Source/NiagaraFluids/Private/ShallowWaterComponent.cpp
@@ -0,0 +1,77 @@
+﻿// Copyright Epic Games, Inc. All Rights Reserved.
+
+#include "ShallowWaterComponent.h"
+
+#include "ShallowWaterSettings.h"
+#include "NiagaraComponent.h"
+#include "NiagaraFunctionLibrary.h"
+#include "WaterBodyActor.h"
+#include "WaterSubsystem.h"
+#include "Components/SkeletalMeshComponent.h"
+#include "Engine/TextureRenderTarget2D.h"
+
+
+UShallowWaterComponent::UShallowWaterComponent()
+{
+	PrimaryComponentTick.bCanEverTick = true;
+	SetComponentTickEnabled(true);
+	SetUsingAbsoluteRotation(true);
+}
+
+void UShallowWaterComponent::BeginPlay()
+{
+	Super::BeginPlay();
+
+	GetOwner()->OnActorBeginOverlap.AddDynamic(this, &UShallowWaterComponent::OnOwnerBeginOverlap);
+	TArray<USkeletalMeshComponent*> SKMs;
+	GetOwner()->GetComponents<USkeletalMeshComponent>(SKMs);
+	for(USkeletalMeshComponent* SKM : SKMs)
+	{
+		SKM->ComponentTags.Add(FName("RigidMesh_ShallowWaterCollider"));
+	}
+
+	ShallowWaterNiagaraSimulation = UNiagaraFunctionLibrary::SpawnSystemAttached(GetDefault<UShallowWaterSettings>()->DefaultShallowWaterNiagaraSimulation.LoadSynchronous(), this, FName(""),
+		FVector::Zero(), FRotator::ZeroRotator, EAttachLocation::SnapToTarget, false);
+	ShallowWaterNiagaraSimulation->SetGpuComputeDebug(GetDefault<UShallowWaterSettings>()->bGPUComputeDebug);
+}
+
+void UShallowWaterComponent::TickComponent(float DeltaTime, ELevelTick TickType,
+                                           FActorComponentTickFunction* ThisTickFunction)
+{
+	Super::TickComponent(DeltaTime, TickType, ThisTickFunction);
+
+	if(CurrentWaterBody == nullptr)
+	{
+		return;
+	}
+
+	// #todo this query gets the closest water surface in 3D space, not a XY-plane projection as we instinctively assumed
+	// Could be better and could be worse
+	const FWaterBodyQueryResult WaterInfo = CurrentWaterBody->GetWaterBodyComponent()->QueryWaterInfoClosestToWorldLocation(
+		GetOwner()->GetActorLocation(), EWaterBodyQueryFlags::ComputeLocation);
+	const FVector WaterLocation = WaterInfo.GetWaterSurfaceLocation();
+	SetWorldLocation(FVector(GetOwner()->GetActorLocation().X, GetOwner()->GetActorLocation().Y, WaterLocation.Z));
+}
+
+void UShallowWaterComponent::OnRegister()
+{
+	Super::OnRegister();
+}
+
+void UShallowWaterComponent::OnOwnerBeginOverlap(AActor* Owner, AActor* OtherActor)
+{
+	AWaterBody* OtherWaterBody = Cast<AWaterBody>(OtherActor);
+	if(OtherWaterBody != nullptr && OtherWaterBody != CurrentWaterBody)
+	{
+		CurrentWaterBody = OtherWaterBody;
+		UpdateNiagaraVariables();
+	}
+}
+
+void UShallowWaterComponent::UpdateNiagaraVariables()
+{
+	TObjectPtr<UTextureRenderTarget2D> WaterInfoTexture = CurrentWaterBody->GetWaterBodyComponent()->GetWaterZone()->
+	                                                                        WaterInfoTexture;
+	ShallowWaterNiagaraSimulation->SetVariableTexture(FName("WaterInfoTexture"), Cast<UTexture>(WaterInfoTexture));
+}
+
diff --git a/Engine/Plugins/FX/NiagaraFluids/Source/NiagaraFluids/Private/ShallowWaterSettings.cpp b/Engine/Plugins/FX/NiagaraFluids/Source/NiagaraFluids/Private/ShallowWaterSettings.cpp
new file mode 100644
--- /dev/null
+++ b/Engine/Plugins/FX/NiagaraFluids/Source/NiagaraFluids/Private/ShallowWaterSettings.cpp
@@ -0,0 +1,16 @@
+﻿// Copyright Epic Games, Inc. All Rights Reserved.
+
+#include "ShallowWaterSettings.h"
+
+UShallowWaterSettings::UShallowWaterSettings()
+{
+
+	
+	DefaultShallowWaterNiagaraSimulation = TSoftObjectPtr<UNiagaraSystem>(FSoftObjectPath(TEXT("/NiagaraFluids/Templates/Liquid/2D/Systems/ShallowWater/Grid2D_SW_FollowPlayer.Grid2D_SW_FollowPlayer")));
+
+}
+
+FName UShallowWaterSettings::GetCategoryName() const
+{
+	return FName("Plugins");
+}
diff --git a/Engine/Plugins/FX/NiagaraFluids/Source/NiagaraFluids/Public/ShallowWaterComponent.h b/Engine/Plugins/FX/NiagaraFluids/Source/NiagaraFluids/Public/ShallowWaterComponent.h
new file mode 100644
--- /dev/null
+++ b/Engine/Plugins/FX/NiagaraFluids/Source/NiagaraFluids/Public/ShallowWaterComponent.h
@@ -0,0 +1,35 @@
+﻿// Copyright Epic Games, Inc. All Rights Reserved.
+
+#pragma once
+
+#include "CoreMinimal.h"
+#include "Components/SceneComponent.h"
+#include "ShallowWaterComponent.generated.h"
+
+class AWaterBody;
+class UNiagaraComponent;
+
+
+UCLASS(ClassGroup=(Custom), meta=(BlueprintSpawnableComponent))
+class NIAGARAFLUIDS_API UShallowWaterComponent : public USceneComponent
+{
+	GENERATED_BODY()
+
+public:
+	UShallowWaterComponent();
+	virtual void BeginPlay() override;
+	virtual void TickComponent(float DeltaTime, ELevelTick TickType,
+							   FActorComponentTickFunction* ThisTickFunction) override;
+	virtual void OnRegister() override;
+	
+protected:
+	// Asset can be set in Project Settings - Plugins - Water ShallowWaterSimulation
+	UPROPERTY(BlueprintReadOnly, VisibleAnywhere, Category="Shallow Water")
+	TObjectPtr<UNiagaraComponent> ShallowWaterNiagaraSimulation;
+
+	UFUNCTION()
+	void OnOwnerBeginOverlap(AActor* Owner, AActor* OtherActor);
+	void UpdateNiagaraVariables();
+	
+	AWaterBody* CurrentWaterBody = nullptr; 
+};
diff --git a/Engine/Plugins/FX/NiagaraFluids/Source/NiagaraFluids/Public/ShallowWaterSettings.h b/Engine/Plugins/FX/NiagaraFluids/Source/NiagaraFluids/Public/ShallowWaterSettings.h
new file mode 100644
--- /dev/null
+++ b/Engine/Plugins/FX/NiagaraFluids/Source/NiagaraFluids/Public/ShallowWaterSettings.h
@@ -0,0 +1,24 @@
+﻿// Copyright Epic Games, Inc. All Rights Reserved.
+
+#pragma once
+
+#include "CoreMinimal.h"
+#include "Engine/DeveloperSettings.h"
+#include "NiagaraSystem.h"
+#include "ShallowWaterSettings.generated.h"
+
+
+UCLASS(config = Engine, defaultconfig, meta=(DisplayName="Water Shallow Water"))
+class NIAGARAFLUIDS_API UShallowWaterSettings : public UDeveloperSettings
+{
+	GENERATED_BODY()
+	UShallowWaterSettings();
+	
+public:
+	virtual FName GetCategoryName() const override;
+	
+	UPROPERTY(Config, BlueprintReadOnly, EditAnywhere, Category="Shallow Water Simulation")
+	TSoftObjectPtr<UNiagaraSystem> DefaultShallowWaterNiagaraSimulation;
+	UPROPERTY(Config, BlueprintReadOnly, EditAnywhere, Category="Shallow Water Simulation")
+	bool bGPUComputeDebug = false;
+};
