[Backout] - CL29603921
Original CL Desc
-----------------------------------------------------------------
NiagaraFluids: Add ShallowWater-WaterBody integration WIP

- Niagara system
- WaterBody context Helper

diff --git a/Engine/Plugins/FX/NiagaraFluids/NiagaraFluids.uplugin b/Engine/Plugins/FX/NiagaraFluids/NiagaraFluids.uplugin
--- a/Engine/Plugins/FX/NiagaraFluids/NiagaraFluids.uplugin
+++ b/Engine/Plugins/FX/NiagaraFluids/NiagaraFluids.uplugin
@@ -1,5 +1,5 @@
 {
-	"FileVersion": 3,
+	"FileVersion": 1,
 	"Version": 1,
 	"VersionName": "1.0",
 	"FriendlyName": "NiagaraFluids",
@@ -20,15 +20,11 @@
 			"Type": "Runtime",
 			"LoadingPhase": "Default"
 		}
-	],
+	],    
 	"Plugins": [
-		{
-			"Name": "Niagara",
-			"Enabled": true
-		},
-		{
-			"Name": "Water",
-			"Enabled": true
-		}
+        {
+            "Name": "Niagara",			
+            "Enabled": true
+        }
 	]
 }
\ No newline at end of file
diff --git a/Engine/Plugins/FX/NiagaraFluids/Source/NiagaraFluids/NiagaraFluids.Build.cs b/Engine/Plugins/FX/NiagaraFluids/Source/NiagaraFluids/NiagaraFluids.Build.cs
--- a/Engine/Plugins/FX/NiagaraFluids/Source/NiagaraFluids/NiagaraFluids.Build.cs
+++ b/Engine/Plugins/FX/NiagaraFluids/Source/NiagaraFluids/NiagaraFluids.Build.cs
@@ -7,18 +7,8 @@ namespace UnrealBuildTool.Rules
 		public NiagaraFluids(ReadOnlyTargetRules Target) : base(Target)
 		{
 			PublicDependencyModuleNames.Add("Core");
-			PrivateDependencyModuleNames.AddRange(
-            			new string[] {
-				            "Core",
-				            "CoreUObject",
-				            "Engine",
-				            "RenderCore",
-            				"Niagara",
-            				"Water",
-            				"DeveloperSettings",
-				            "Projects"
-            			}
-            		);
+			PrivateDependencyModuleNames.Add("RenderCore");
+			PrivateDependencyModuleNames.Add("Projects");
 		}
 	}
 }
diff --git a/Engine/Plugins/FX/NiagaraFluids/Source/NiagaraFluids/Private/ShallowWaterComponent.cpp b/Engine/Plugins/FX/NiagaraFluids/Source/NiagaraFluids/Private/ShallowWaterComponent.cpp
deleted file mode 100644
diff --git a/Engine/Plugins/FX/NiagaraFluids/Source/NiagaraFluids/Private/ShallowWaterSettings.cpp b/Engine/Plugins/FX/NiagaraFluids/Source/NiagaraFluids/Private/ShallowWaterSettings.cpp
deleted file mode 100644
diff --git a/Engine/Plugins/FX/NiagaraFluids/Source/NiagaraFluids/Public/ShallowWaterComponent.h b/Engine/Plugins/FX/NiagaraFluids/Source/NiagaraFluids/Public/ShallowWaterComponent.h
deleted file mode 100644
diff --git a/Engine/Plugins/FX/NiagaraFluids/Source/NiagaraFluids/Public/ShallowWaterSettings.h b/Engine/Plugins/FX/NiagaraFluids/Source/NiagaraFluids/Public/ShallowWaterSettings.h
deleted file mode 100644
